Disable All_Android_On_ComputerName from showing up for MagicLeap/Lumin/Phaedra.

(Proper change to the name requires changing IniPlatformName() for LuminTargetPlatform, but that is used in several other places in the code, and I'm not sure if there will be other impacts, so just disabling that menu for now. Main concern was it just being confusing saying "Android", so it no being there at all for 4.20 should be fine.)

diff --git a/Engine/Source/Developer/Lumin/LuminTargetPlatform/Private/LuminTargetPlatform.h b/Engine/Source/Developer/Lumin/LuminTargetPlatform/Private/LuminTargetPlatform.h
--- a/Engine/Source/Developer/Lumin/LuminTargetPlatform/Private/LuminTargetPlatform.h
+++ b/Engine/Source/Developer/Lumin/LuminTargetPlatform/Private/LuminTargetPlatform.h
@@ -54,6 +54,11 @@ public:
 		return TEXT("Lumin");
 	}
 
+	virtual bool SupportsVariants() const override
+	{
+		return false;
+	}
+
 	virtual bool IsSdkInstalled(bool bProjectHasCode, FString& OutDocumentationPath) const override;
 
 	virtual int32 CheckRequirements(const FString& ProjectPath, bool bProjectHasCode, FString& OutTutorialPath, FString& OutDocumentationPath, FText& CustomizedLogMessage) const override;
