Disable All_Android_On_ComputerName from showing up for MagicLeap/Lumin/Phaedra.

(Proper change to the name requires changing IniPlatformName() for LuminTargetPlatform, but that is used in several other places in the code, and I'm not sure if there will be other impacts, so just disabling that menu for now. Main concern was it just being confusing saying "Android", so it no being there at all for 4.20 should be fine.)

diff --git a/Engine/Source/Developer/Android/AndroidTargetPlatform/Private/AndroidTargetPlatform.cpp b/Engine/Source/Developer/Android/AndroidTargetPlatform/Private/AndroidTargetPlatform.cpp
--- a/Engine/Source/Developer/Android/AndroidTargetPlatform/Private/AndroidTargetPlatform.cpp
+++ b/Engine/Source/Developer/Android/AndroidTargetPlatform/Private/AndroidTargetPlatform.cpp
@@ -778,7 +778,8 @@ bool FAndroidTargetPlatform::HandleTicker( float DeltaTime )
 			// create target device
 			FAndroidTargetDevicePtr& Device = Devices.Add(DeviceInfo.SerialNumber);
 
-			Device = MakeShareable(new FAndroidTargetDevice(*this, DeviceInfo.SerialNumber, GetAndroidVariantName()));
+			Device = CreateNewDevice(DeviceInfo);
+
 
 			Device->SetConnected(true);
 			Device->SetModel(DeviceInfo.Model);
@@ -807,4 +808,9 @@ bool FAndroidTargetPlatform::HandleTicker( float DeltaTime )
 	return true;
 }
 
+FAndroidTargetDeviceRef FAndroidTargetPlatform::CreateNewDevice(const FAndroidDeviceInfo &DeviceInfo)
+{
+	return MakeShareable(new FAndroidTargetDevice(*this, DeviceInfo.SerialNumber, GetAndroidVariantName()));
+}
+
 #undef LOCTEXT_NAMESPACE
diff --git a/Engine/Source/Developer/Android/AndroidTargetPlatform/Private/AndroidTargetPlatform.h b/Engine/Source/Developer/Android/AndroidTargetPlatform/Private/AndroidTargetPlatform.h
--- a/Engine/Source/Developer/Android/AndroidTargetPlatform/Private/AndroidTargetPlatform.h
+++ b/Engine/Source/Developer/Android/AndroidTargetPlatform/Private/AndroidTargetPlatform.h
@@ -281,6 +281,8 @@ protected:
 	// Handles when the ticker fires.
 	bool HandleTicker( float DeltaTime );
 
+	virtual FAndroidTargetDeviceRef CreateNewDevice(const FAndroidDeviceInfo &DeviceInfo);
+
 	// true if this is a client TP
 	bool bIsClient;
 
diff --git a/Engine/Source/Developer/Lumin/LuminTargetPlatform/Private/LuminTargetPlatform.cpp b/Engine/Source/Developer/Lumin/LuminTargetPlatform/Private/LuminTargetPlatform.cpp
--- a/Engine/Source/Developer/Lumin/LuminTargetPlatform/Private/LuminTargetPlatform.cpp
+++ b/Engine/Source/Developer/Lumin/LuminTargetPlatform/Private/LuminTargetPlatform.cpp
@@ -116,6 +116,11 @@ bool FLuminTargetPlatform::SupportsFeature(ETargetPlatformFeatures Feature) cons
 	return TTargetPlatformBase<FAndroidPlatformProperties>::SupportsFeature(Feature);
 }
 
+FAndroidTargetDeviceRef FLuminTargetPlatform::CreateNewDevice(const FAndroidDeviceInfo &DeviceInfo)
+{
+	return MakeShareable(new FLuminTargetDevice(*this, DeviceInfo.SerialNumber, GetAndroidVariantName()));
+}
+
 void FLuminTargetPlatform::InitializeDeviceDetection()
 {
 	DeviceDetection = FModuleManager::LoadModuleChecked<IAndroidDeviceDetectionModule>("AndroidDeviceDetection").GetAndroidDeviceDetection(TEXT("Lumin"));
diff --git a/Engine/Source/Developer/Lumin/LuminTargetPlatform/Private/LuminTargetPlatform.h b/Engine/Source/Developer/Lumin/LuminTargetPlatform/Private/LuminTargetPlatform.h
--- a/Engine/Source/Developer/Lumin/LuminTargetPlatform/Private/LuminTargetPlatform.h
+++ b/Engine/Source/Developer/Lumin/LuminTargetPlatform/Private/LuminTargetPlatform.h
@@ -28,6 +28,31 @@ class ILuminDeviceDetection;
 class UTextureLODSettings;
 enum class ETargetPlatformFeatures;
 
+/**
+* Implements a Lumin target device.
+*/
+class FLuminTargetDevice : public FAndroidTargetDevice
+{
+public:
+
+	/**
+	* Creates and initializes a new Lumin target device.
+	*
+	* @param InTargetPlatform - The target platform.
+	* @param InSerialNumber - The ADB serial number of the target device.
+	* @param InAndroidVariant - The variant of the Android platform, i.e. ATC, DXT or PVRTC.
+	*/
+	FLuminTargetDevice(const ITargetPlatform& InTargetPlatform, const FString& InSerialNumber, const FString& InAndroidVariant)
+		: FAndroidTargetDevice(InTargetPlatform, InSerialNumber, InAndroidVariant)
+	{ }
+
+	// Return true if the devices can be grouped in an aggregate (All_<platform>_devices_on_<host>) proxy
+	virtual bool IsPlatformAggregated() const override
+	{
+		return false;
+	}
+};
+
 /**
  * FSeakTargetPlatform, abstraction for cooking Lumin platforms
  */
@@ -90,6 +115,8 @@ public:
 	virtual bool SupportsFeature(ETargetPlatformFeatures Feature) const override;
 
 protected:
+	virtual FAndroidTargetDeviceRef CreateNewDevice(const FAndroidDeviceInfo &DeviceInfo);
+
 	// Holds the Engine INI settings, for quick use.
 	FConfigFile EngineSettings;
 
