Moving Lumin SDK Settings back into per project from per user

diff --git a/Engine/Source/Developer/Lumin/LuminPlatformEditor/Private/MagicLeapSDKSettings.h b/Engine/Source/Developer/Lumin/LuminPlatformEditor/Private/MagicLeapSDKSettings.h
--- a/Engine/Source/Developer/Lumin/LuminPlatformEditor/Private/MagicLeapSDKSettings.h
+++ b/Engine/Source/Developer/Lumin/LuminPlatformEditor/Private/MagicLeapSDKSettings.h
@@ -11,7 +11,7 @@ class IAndroidDeviceDetection;
 /**
  * Implements the settings for the Magic Leap SDK setup.
  */
-UCLASS(config = Engine, globaluserconfig)
+UCLASS(config = Engine, defaultconfig)
 class LUMINPLATFORMEDITOR_API UMagicLeapSDKSettings : public UObject
 {
 public:
