VR - Fix Lumin not including lightmaps

diff --git a/Engine/Source/Developer/Lumin/LuminTargetPlatform/Private/LuminTargetPlatform.cpp b/Engine/Source/Developer/Lumin/LuminTargetPlatform/Private/LuminTargetPlatform.cpp
--- a/Engine/Source/Developer/Lumin/LuminTargetPlatform/Private/LuminTargetPlatform.cpp
+++ b/Engine/Source/Developer/Lumin/LuminTargetPlatform/Private/LuminTargetPlatform.cpp
@@ -87,6 +87,33 @@ static bool LuminSupportsVulkan(const FConfigFile& EngineSettings)
 	return bSupportsVulkan;
 }
 
+bool FLuminTargetPlatform::SupportsFeature(ETargetPlatformFeatures Feature) const
+{
+	switch (Feature)
+	{
+	case ETargetPlatformFeatures::Packaging:
+		return true;
+
+	case ETargetPlatformFeatures::LowQualityLightmaps:
+	case ETargetPlatformFeatures::MobileRendering:
+		return SupportsMobileRendering() || SupportsVulkan();
+
+	case ETargetPlatformFeatures::HighQualityLightmaps:
+	//#todo-rco: Enable when Vulkan supports it
+	//case ETargetPlatformFeatures::Tessellation:
+	case ETargetPlatformFeatures::DeferredRendering:
+		return SupportsDesktopRendering();
+
+	case ETargetPlatformFeatures::SoftwareOcclusion:
+		return SupportsSoftwareOcclusion();
+
+	default:
+		break;
+	}
+
+	return TTargetPlatformBase<FAndroidPlatformProperties>::SupportsFeature(Feature);
+}
+
 void FLuminTargetPlatform::InitializeDeviceDetection()
 {
 	DeviceDetection = FModuleManager::LoadModuleChecked<IAndroidDeviceDetectionModule>("AndroidDeviceDetection").GetAndroidDeviceDetection(TEXT("Lumin"));
diff --git a/Engine/Source/Developer/Lumin/LuminTargetPlatform/Private/LuminTargetPlatform.h b/Engine/Source/Developer/Lumin/LuminTargetPlatform/Private/LuminTargetPlatform.h
--- a/Engine/Source/Developer/Lumin/LuminTargetPlatform/Private/LuminTargetPlatform.h
+++ b/Engine/Source/Developer/Lumin/LuminTargetPlatform/Private/LuminTargetPlatform.h
@@ -87,6 +87,8 @@ public:
 	virtual bool SupportsMobileRendering() const;// override;
 	virtual void InitializeDeviceDetection() override;
 
+	virtual bool SupportsFeature(ETargetPlatformFeatures Feature) const override;
+
 protected:
 	// Holds the Engine INI settings, for quick use.
 	FConfigFile EngineSettings;
