VR - Fix incorrect Lumin shader platform

diff --git a/Engine/Source/Developer/Lumin/LuminTargetPlatform/Private/LuminTargetPlatform.cpp b/Engine/Source/Developer/Lumin/LuminTargetPlatform/Private/LuminTargetPlatform.cpp
--- a/Engine/Source/Developer/Lumin/LuminTargetPlatform/Private/LuminTargetPlatform.cpp
+++ b/Engine/Source/Developer/Lumin/LuminTargetPlatform/Private/LuminTargetPlatform.cpp
@@ -30,7 +30,7 @@ FLuminTargetPlatform::FLuminTargetPlatform(bool bIsClient)
 	GetAllTargetedShaderFormats(TargetedShaderFormats);
 
 	// If we are targeting ES 2.0/3.1, we also must cook encoded HDR reflection captures
-	static FName NAME_VULKAN_ES31(TEXT("SF_VULKAN_ES31_ANDROID"));
+	static FName NAME_VULKAN_ES31(TEXT("SF_VULKAN_ES31_LUMIN"));
 	static FName NAME_GLSL_ES2(TEXT("GLSL_ES2"));
 	static FName NAME_GLSL_SM5(TEXT("GLSL_430"));
 	bRequiresEncodedHDRReflectionCaptures = TargetedShaderFormats.Contains(NAME_VULKAN_ES31)
diff --git a/Engine/Source/Runtime/VulkanRHI/Private/VulkanRHI.cpp b/Engine/Source/Runtime/VulkanRHI/Private/VulkanRHI.cpp
--- a/Engine/Source/Runtime/VulkanRHI/Private/VulkanRHI.cpp
+++ b/Engine/Source/Runtime/VulkanRHI/Private/VulkanRHI.cpp
@@ -61,7 +61,7 @@ FDynamicRHI* FVulkanDynamicRHIModule::CreateRHI(ERHIFeatureLevel::Type InRequest
 			FParse::Param(FCommandLine::Get(), TEXT("featureleveles31")) || FParse::Param(FCommandLine::Get(), TEXT("featureleveles2"))))
 	{
 		GMaxRHIFeatureLevel = ERHIFeatureLevel::ES3_1;
-		GMaxRHIShaderPlatform = PLATFORM_ANDROID ? SP_VULKAN_ES3_1_ANDROID : (PLATFORM_LUMIN ?  SP_VULKAN_ES3_1_LUMIN : SP_VULKAN_PCES3_1);
+		GMaxRHIShaderPlatform = PLATFORM_LUMIN ? SP_VULKAN_ES3_1_LUMIN : (PLATFORM_ANDROID ? SP_VULKAN_ES3_1_ANDROID : SP_VULKAN_PCES3_1);
 	}
 	else if (InRequestedFeatureLevel == ERHIFeatureLevel::SM4)
 	{
