Magic Leap: Override IniPlatfromName() for LuminTargetPlatform to report "Lumin" rather than reporting the value from it's parent, AndroidTargetPlatform, "Android".

Unsure if this will affect more than just this display name, as IniPlatformName() is used in a few places in the code, but eventually it will need to be "Lumin" anyway so we'll fix fallout as we find it.

diff --git a/Engine/Source/Developer/Lumin/LuminTargetPlatform/Private/LuminTargetPlatform.h b/Engine/Source/Developer/Lumin/LuminTargetPlatform/Private/LuminTargetPlatform.h
--- a/Engine/Source/Developer/Lumin/LuminTargetPlatform/Private/LuminTargetPlatform.h
+++ b/Engine/Source/Developer/Lumin/LuminTargetPlatform/Private/LuminTargetPlatform.h
@@ -54,6 +54,11 @@ public:
 		return TEXT("Lumin");
 	}
 
+	virtual FString IniPlatformName() const override
+	{
+		return TEXT("Lumin");
+	}
+
 	virtual bool IsSdkInstalled(bool bProjectHasCode, FString& OutDocumentationPath) const override;
 
 	virtual int32 CheckRequirements(const FString& ProjectPath, bool bProjectHasCode, FString& OutTutorialPath, FString& OutDocumentationPath, FText& CustomizedLogMessage) const override;
