For Lumin, only build shaders for OpenGL or Vulkan, not both. Fix a shader compile error in BogMyrtle material (can't use SpeedTree node on mobile).

diff --git a/Engine/Source/Developer/Lumin/LuminTargetPlatform/Private/LuminTargetPlatform.cpp b/Engine/Source/Developer/Lumin/LuminTargetPlatform/Private/LuminTargetPlatform.cpp
--- a/Engine/Source/Developer/Lumin/LuminTargetPlatform/Private/LuminTargetPlatform.cpp
+++ b/Engine/Source/Developer/Lumin/LuminTargetPlatform/Private/LuminTargetPlatform.cpp
@@ -149,23 +149,26 @@ void FLuminTargetPlatform::GetAllPossibleShaderFormats( TArray<FName>& OutFormat
 
 	if (SupportsMobileRendering())
 	{
-		OutFormats.AddUnique(NAME_GLSL_ES2);
 		if (LuminSupportsVulkan(EngineSettings))
 		{
 			OutFormats.AddUnique(NAME_VULKAN_ES31_LUMIN);
 		}
+		else
+		{
+			OutFormats.AddUnique(NAME_GLSL_ES2);
+		}
 	}
 
 	if (SupportsDesktopRendering())
 	{
-//		OutFormats.AddUnique(NAME_GLSL_310_ES_EXT);
-//		OutFormats.AddUnique(NAME_GLSL_SM4);
-		OutFormats.AddUnique(NAME_GLSL_SM5);
-
 		if (LuminSupportsVulkan(EngineSettings))
 		{
 			OutFormats.AddUnique(NAME_VULKAN_SM5_LUMIN);
 		}
+		else
+		{
+			OutFormats.AddUnique(NAME_GLSL_SM5);
+		}
 	}
 }
 
