CIS build fix - removing duplicate definition & adding a missing #pragma once to a header fille

#mlqdr

diff --git a/Engine/Source/Developer/Lumin/LuminTargetPlatform/Private/LuminTargetDevice.h b/Engine/Source/Developer/Lumin/LuminTargetPlatform/Private/LuminTargetDevice.h
--- a/Engine/Source/Developer/Lumin/LuminTargetPlatform/Private/LuminTargetDevice.h
+++ b/Engine/Source/Developer/Lumin/LuminTargetPlatform/Private/LuminTargetDevice.h
@@ -1,14 +1,28 @@
 // Copyright 1998-2018 Epic Games, Inc. All Rights Reserved.
 
+#pragma once
+
 #include "AndroidTargetDevice.h"
 
+/**
+* Implements a Lumin target device.
+*/
 class FLuminTargetDevice : public FAndroidTargetDevice
 {
 public:
+
+	/**
+	* Creates and initializes a new Lumin target device.
+	*
+	* @param InTargetPlatform - The target platform.
+	* @param InSerialNumber - The ADB serial number of the target device.
+	* @param InAndroidVariant - The variant of the Android platform, i.e. ATC, DXT or PVRTC.
+	*/
 	FLuminTargetDevice(const ITargetPlatform& InTargetPlatform, const FString& InSerialNumber, const FString& InAndroidVariant)
-	: FAndroidTargetDevice(InTargetPlatform, InSerialNumber, InAndroidVariant)
-	{}
+		: FAndroidTargetDevice(InTargetPlatform, InSerialNumber, InAndroidVariant)
+	{ }
 
+	// Return true if the devices can be grouped in an aggregate (All_<platform>_devices_on_<host>) proxy
 	virtual bool IsPlatformAggregated() const override
 	{
 		return false;
diff --git a/Engine/Source/Developer/Lumin/LuminTargetPlatform/Private/LuminTargetPlatform.h b/Engine/Source/Developer/Lumin/LuminTargetPlatform/Private/LuminTargetPlatform.h
--- a/Engine/Source/Developer/Lumin/LuminTargetPlatform/Private/LuminTargetPlatform.h
+++ b/Engine/Source/Developer/Lumin/LuminTargetPlatform/Private/LuminTargetPlatform.h
@@ -28,31 +28,6 @@ class ILuminDeviceDetection;
 class UTextureLODSettings;
 enum class ETargetPlatformFeatures;
 
-/**
-* Implements a Lumin target device.
-*/
-class FLuminTargetDevice : public FAndroidTargetDevice
-{
-public:
-
-	/**
-	* Creates and initializes a new Lumin target device.
-	*
-	* @param InTargetPlatform - The target platform.
-	* @param InSerialNumber - The ADB serial number of the target device.
-	* @param InAndroidVariant - The variant of the Android platform, i.e. ATC, DXT or PVRTC.
-	*/
-	FLuminTargetDevice(const ITargetPlatform& InTargetPlatform, const FString& InSerialNumber, const FString& InAndroidVariant)
-		: FAndroidTargetDevice(InTargetPlatform, InSerialNumber, InAndroidVariant)
-	{ }
-
-	// Return true if the devices can be grouped in an aggregate (All_<platform>_devices_on_<host>) proxy
-	virtual bool IsPlatformAggregated() const override
-	{
-		return false;
-	}
-};
-
 /**
  * FSeakTargetPlatform, abstraction for cooking Lumin platforms
  */
