#DEV-VR - Fix for swapping shader cache formats and using Launch On

We now reload the settings each time they are used rather than caching once at startup

diff --git a/Engine/Source/Developer/Lumin/LuminTargetPlatform/Private/LuminTargetPlatform.cpp b/Engine/Source/Developer/Lumin/LuminTargetPlatform/Private/LuminTargetPlatform.cpp
--- a/Engine/Source/Developer/Lumin/LuminTargetPlatform/Private/LuminTargetPlatform.cpp
+++ b/Engine/Source/Developer/Lumin/LuminTargetPlatform/Private/LuminTargetPlatform.cpp
@@ -8,6 +8,7 @@
 #include "UObject/NameTypes.h"
 #include "Logging/LogMacros.h"
 #include "Stats/Stats.h"
+#include "Materials/Material.h"
 #include "LuminTargetDevice.h"
 #include "Modules/ModuleManager.h"
 
@@ -26,19 +27,7 @@ FLuminTargetPlatform::FLuminTargetPlatform(bool bIsClient)
 	// by using the FAndroidPlatformProperties, the PlatformInfo up in TTargetPlatformBase/FTargetPlatformBase would be Android
 	this->PlatformInfo = PlatformInfo::FindPlatformInfo("Lumin");
 
-	FConfigCacheIni::LoadLocalIniFile(EngineSettings, TEXT("Engine"), true, *IniPlatformName());
-
-	// Get the Target RHIs for this platform, we do not always want all those that are supported.
-	TArray<FName> TargetedShaderFormats;
-	GetAllTargetedShaderFormats(TargetedShaderFormats);
-
-	// If we are targeting ES 2.0/3.1, we also must cook encoded HDR reflection captures
-	static FName NAME_VULKAN_ES31(TEXT("SF_VULKAN_ES31_LUMIN"));
-	static FName NAME_GLSL_ES2(TEXT("GLSL_ES2"));
-	static FName NAME_GLSL_SM5(TEXT("GLSL_430"));
-	bRequiresEncodedHDRReflectionCaptures = TargetedShaderFormats.Contains(NAME_VULKAN_ES31)
-		|| TargetedShaderFormats.Contains(NAME_GLSL_ES2)
-		|| TargetedShaderFormats.Contains(NAME_GLSL_SM5);
+	RefreshSettings();
 #endif
 }
 
@@ -72,24 +61,62 @@ int32 FLuminTargetPlatform::CheckRequirements(const FString& ProjectPath, bool b
 bool FLuminTargetPlatform::SupportsMobileRendering() const
 {
 	bool bUseMobileRendering = true;
-	EngineSettings.GetBool(TEXT("/Script/LuminRuntimeSettings.LuminRuntimeSettings"), TEXT("bUseMobileRendering"), bUseMobileRendering);
+	LuminEngineSettings.GetBool(TEXT("/Script/LuminRuntimeSettings.LuminRuntimeSettings"), TEXT("bUseMobileRendering"), bUseMobileRendering);
 	return bUseMobileRendering;
 }
 
 bool FLuminTargetPlatform::SupportsDesktopRendering() const
 {
 	bool bUseMobileRendering = true;
-	EngineSettings.GetBool(TEXT("/Script/LuminRuntimeSettings.LuminRuntimeSettings"), TEXT("bUseMobileRendering"), bUseMobileRendering);
+	LuminEngineSettings.GetBool(TEXT("/Script/LuminRuntimeSettings.LuminRuntimeSettings"), TEXT("bUseMobileRendering"), bUseMobileRendering);
 	return bUseMobileRendering == false;
 }
 
-static bool LuminSupportsVulkan(const FConfigFile& EngineSettings)
+static bool LuminSupportsVulkan(const FConfigFile& InLuminEngineSettings)
 {
 	bool bSupportsVulkan = false;
-	EngineSettings.GetBool(TEXT("/Script/LuminRuntimeSettings.LuminRuntimeSettings"), TEXT("bUseVulkan"), bSupportsVulkan);
+	InLuminEngineSettings.GetBool(TEXT("/Script/LuminRuntimeSettings.LuminRuntimeSettings"), TEXT("bUseVulkan"), bSupportsVulkan);
 	return bSupportsVulkan;
 }
 
+void FLuminTargetPlatform::RefreshSettings()
+{
+#if WITH_ENGINE
+	//UE_LOG(LogCore, Warning, TEXT("*** DIAGNOSE - REFRESH!!!"));
+
+	FConfigFile NewEngineSettings;
+	FConfigCacheIni::LoadLocalIniFile(NewEngineSettings, TEXT("Engine"), true, *IniPlatformName(), true);
+	//the load above does not move settings from "SourceConfig" member to the object itself.  New loads will do that.
+	LuminEngineSettings = NewEngineSettings;
+	//override the android version too
+	EngineSettings = NewEngineSettings;
+
+	// Get the Target RHIs for this platform, we do not always want all those that are supported.
+	TArray<FName> TargetedShaderFormats;
+	GetAllTargetedShaderFormats(TargetedShaderFormats);
+
+	// If we are targeting ES 2.0/3.1, we also must cook encoded HDR reflection captures
+	static FName NAME_VULKAN_ES31(TEXT("SF_VULKAN_ES31_LUMIN"));
+	static FName NAME_GLSL_ES2(TEXT("GLSL_ES2"));
+	static FName NAME_GLSL_SM5(TEXT("GLSL_430"));
+	bRequiresEncodedHDRReflectionCaptures = TargetedShaderFormats.Contains(NAME_VULKAN_ES31)
+		|| TargetedShaderFormats.Contains(NAME_GLSL_ES2)
+		|| TargetedShaderFormats.Contains(NAME_GLSL_SM5);
+
+	//ensure that we wipe out the material cached data before we begin serializing.  It is cleared *after* a serialize, but changes made ini files will not be taken into account for materials without this
+	TArray<UObject*> Materials;
+	GetObjectsOfClass(UMaterial::StaticClass(), Materials, true);
+	for (UObject* Material : Materials)
+	{
+		if (Material->GetOutermost() != GetTransientPackage())
+		{
+			Material->ClearCachedCookedPlatformData(this);
+		}
+	}
+
+#endif
+}
+
 bool FLuminTargetPlatform::SupportsFeature(ETargetPlatformFeatures Feature) const
 {
 	switch (Feature)
@@ -149,7 +176,7 @@ void FLuminTargetPlatform::GetAllPossibleShaderFormats( TArray<FName>& OutFormat
 
 	if (SupportsMobileRendering())
 	{
-		if (LuminSupportsVulkan(EngineSettings))
+		if (LuminSupportsVulkan(LuminEngineSettings))
 		{
 			OutFormats.AddUnique(NAME_VULKAN_ES31_LUMIN);
 		}
@@ -161,7 +188,7 @@ void FLuminTargetPlatform::GetAllPossibleShaderFormats( TArray<FName>& OutFormat
 
 	if (SupportsDesktopRendering())
 	{
-		if (LuminSupportsVulkan(EngineSettings))
+		if (LuminSupportsVulkan(LuminEngineSettings))
 		{
 			OutFormats.AddUnique(NAME_VULKAN_SM5_LUMIN);
 		}
@@ -201,7 +228,7 @@ void FLuminTargetPlatform::GetTextureFormats(const UTexture* InTexture, TArray<F
 	// if we didn't assign anything specially, then use the defaults
 	if (TextureFormatName == NAME_None)
 	{
-		TextureFormatName = GetDefaultTextureFormatName(this, InTexture, EngineSettings, false);
+		TextureFormatName = GetDefaultTextureFormatName(this, InTexture, LuminEngineSettings, false);
 	}
 
 	// perform any remapping away from defaults
diff --git a/Engine/Source/Developer/Lumin/LuminTargetPlatform/Private/LuminTargetPlatform.h b/Engine/Source/Developer/Lumin/LuminTargetPlatform/Private/LuminTargetPlatform.h
--- a/Engine/Source/Developer/Lumin/LuminTargetPlatform/Private/LuminTargetPlatform.h
+++ b/Engine/Source/Developer/Lumin/LuminTargetPlatform/Private/LuminTargetPlatform.h
@@ -94,12 +94,14 @@ public:
 
 	virtual bool SupportsFeature(ETargetPlatformFeatures Feature) const override;
 
+	virtual void RefreshSettings() override;
+
 protected:
 	virtual FAndroidTargetDevicePtr CreateTargetDevice(const ITargetPlatform& InTargetPlatform, const FString& InSerialNumber, const FString& InAndroidVariant) const override;
 
 protected:
 	// Holds the Engine INI settings, for quick use.
-	FConfigFile EngineSettings;
+	FConfigFile LuminEngineSettings;
 
 };
 
diff --git a/Engine/Source/Developer/TargetPlatform/Public/Common/TargetPlatformBase.h b/Engine/Source/Developer/TargetPlatform/Public/Common/TargetPlatformBase.h
--- a/Engine/Source/Developer/TargetPlatform/Public/Common/TargetPlatformBase.h
+++ b/Engine/Source/Developer/TargetPlatform/Public/Common/TargetPlatformBase.h
@@ -93,6 +93,10 @@ public:
 		// do nothing in the base class
 	}
 
+	virtual void RefreshSettings() override
+	{
+	}
+
 	virtual int32 GetCompressionBitWindow() const override
 	{
 		return DEFAULT_ZLIB_BIT_WINDOW;
diff --git a/Engine/Source/Developer/TargetPlatform/Public/Interfaces/ITargetPlatform.h b/Engine/Source/Developer/TargetPlatform/Public/Interfaces/ITargetPlatform.h
--- a/Engine/Source/Developer/TargetPlatform/Public/Interfaces/ITargetPlatform.h
+++ b/Engine/Source/Developer/TargetPlatform/Public/Interfaces/ITargetPlatform.h
@@ -433,6 +433,11 @@ public:
 	*/
 	virtual void GetBuildProjectSettingKeys(FString& OutSection, TArray<FString>& InBoolKeys, TArray<FString>& InIntKeys, TArray<FString>& InStringKeys) const = 0;
 
+	/**
+	 * Give the platform a chance to refresh internal settings before a cook, etc
+	 */
+	virtual void RefreshSettings() = 0;
+
 public:
 
 	/**
diff --git a/Engine/Source/Editor/UnrealEd/Classes/CookOnTheSide/CookOnTheFlyServer.h b/Engine/Source/Editor/UnrealEd/Classes/CookOnTheSide/CookOnTheFlyServer.h
--- a/Engine/Source/Editor/UnrealEd/Classes/CookOnTheSide/CookOnTheFlyServer.h
+++ b/Engine/Source/Editor/UnrealEd/Classes/CookOnTheSide/CookOnTheFlyServer.h
@@ -1529,6 +1529,11 @@ private:
 	*/
 	void InitializeSandbox();
 
+	/**
+	* Initialize all target platforms
+	*/
+	void InitializeTargetPlatforms();
+
 	/**
 	* Clean up the sandbox
 	*/
diff --git a/Engine/Source/Editor/UnrealEd/Private/CookOnTheFlyServer.cpp b/Engine/Source/Editor/UnrealEd/Private/CookOnTheFlyServer.cpp
--- a/Engine/Source/Editor/UnrealEd/Private/CookOnTheFlyServer.cpp
+++ b/Engine/Source/Editor/UnrealEd/Private/CookOnTheFlyServer.cpp
@@ -966,6 +966,7 @@ bool UCookOnTheFlyServer::StartNetworkFileServer( const bool BindAnyPort )
 	GenerateAssetRegistry();
 
 	InitializeSandbox();
+	InitializeTargetPlatforms();
 
 	const TArray<ITargetPlatform*>& Platforms = GetCookingTargetPlatforms();
 
@@ -2519,7 +2520,6 @@ void UCookOnTheFlyServer::TickPrecacheObjectsForPlatforms(const float TimeSlice,
 	
 	FCookerTimer Timer(TimeSlice, true);
 
-	++LastUpdateTick;
 	if (LastUpdateTick > 50 ||
 		((CachedMaterialsToCacheArray.Num() == 0) && (CachedTexturesToCacheArray.Num() == 0)))
 	{
@@ -2543,6 +2543,7 @@ void UCookOnTheFlyServer::TickPrecacheObjectsForPlatforms(const float TimeSlice,
 			CachedTexturesToCacheArray.Add(Texture);
 		}
 	}
+	++LastUpdateTick;
 
 	if (Timer.IsTimeUp())
 		return;
@@ -6188,6 +6189,23 @@ void UCookOnTheFlyServer::InitializeSandbox()
 	}
 }
 
+
+void UCookOnTheFlyServer::InitializeTargetPlatforms()
+{
+	const TArray<ITargetPlatform*>& CookingTargetPlatforms = GetCookingTargetPlatforms();
+
+	//allow each platform to update its internals before cooking
+	for (int32 TargetPlatformIndex = 0; TargetPlatformIndex < CookingTargetPlatforms.Num(); ++TargetPlatformIndex)
+	{
+		ITargetPlatform* TargetPlatform = CookingTargetPlatforms[TargetPlatformIndex];
+		if (TargetPlatform)
+		{
+			TargetPlatform->RefreshSettings();
+		}
+	}
+}
+
+
 void UCookOnTheFlyServer::TermSandbox()
 {
 	ClearAllCookedData();
@@ -6226,6 +6244,9 @@ void UCookOnTheFlyServer::StartCookByTheBook( const FCookByTheBookStartupOptions
 	check( IsInGameThread() );
 	check( IsCookByTheBookMode() );
 
+	//force precache objects to refresh themselves before cooking anything
+	LastUpdateTick = INT_MAX;
+
 	CookByTheBookOptions->bRunning = true;
 	CookByTheBookOptions->bCancel = false;
 	CookByTheBookOptions->CookTime = 0.0f;
@@ -6355,6 +6376,7 @@ void UCookOnTheFlyServer::StartCookByTheBook( const FCookByTheBookStartupOptions
 
 	// This will either delete the sandbox or iteratively clean it
 	InitializeSandbox();
+	InitializeTargetPlatforms();
 
 	if (CurrentCookMode == ECookMode::CookByTheBook && !IsCookFlagSet(ECookInitializationFlags::Iterative))
 	{
