#DEV-VR - Fix for Dev-VR compile issue for Lumin Target Platform dependencies on Engine

diff --git a/Engine/Source/Developer/Lumin/LuminTargetPlatform/Private/LuminTargetPlatform.cpp b/Engine/Source/Developer/Lumin/LuminTargetPlatform/Private/LuminTargetPlatform.cpp
--- a/Engine/Source/Developer/Lumin/LuminTargetPlatform/Private/LuminTargetPlatform.cpp
+++ b/Engine/Source/Developer/Lumin/LuminTargetPlatform/Private/LuminTargetPlatform.cpp
@@ -8,7 +8,9 @@
 #include "UObject/NameTypes.h"
 #include "Logging/LogMacros.h"
 #include "Stats/Stats.h"
+#if WITH_EDITOR
 #include "Materials/Material.h"
+#endif
 #include "LuminTargetDevice.h"
 #include "Modules/ModuleManager.h"
 
@@ -103,6 +105,7 @@ void FLuminTargetPlatform::RefreshSettings()
 		|| TargetedShaderFormats.Contains(NAME_GLSL_ES2)
 		|| TargetedShaderFormats.Contains(NAME_GLSL_SM5);
 
+#if WITH_EDITOR
 	//ensure that we wipe out the material cached data before we begin serializing.  It is cleared *after* a serialize, but changes made ini files will not be taken into account for materials without this
 	TArray<UObject*> Materials;
 	GetObjectsOfClass(UMaterial::StaticClass(), Materials, true);
@@ -113,6 +116,7 @@ void FLuminTargetPlatform::RefreshSettings()
 			Material->ClearCachedCookedPlatformData(this);
 		}
 	}
+#endif
 
 #endif
 }
