[Dev-VR] #jira UE-62644 Deescalate resampling logs to display, and make sure that Lumin does not get Resampling settings from Android runtime settings.

diff --git a/Engine/Source/Developer/Lumin/LuminTargetPlatform/Private/LuminTargetPlatform.h b/Engine/Source/Developer/Lumin/LuminTargetPlatform/Private/LuminTargetPlatform.h
--- a/Engine/Source/Developer/Lumin/LuminTargetPlatform/Private/LuminTargetPlatform.h
+++ b/Engine/Source/Developer/Lumin/LuminTargetPlatform/Private/LuminTargetPlatform.h
@@ -76,6 +76,11 @@ public:
 		OutFormats.Add(FName(TEXT("FullHDR")));
 	}
 
+	virtual FPlatformAudioCookOverrides* GetAudioCompressionSettings() const override
+	{
+		return nullptr;
+	}
+	
 	// True if the project requires encoded HDR reflection captures
 	bool bRequiresEncodedHDRReflectionCaptures;
 #endif //WITH_ENGINE
diff --git a/Engine/Source/Runtime/Engine/Private/AudioDecompress.cpp b/Engine/Source/Runtime/Engine/Private/AudioDecompress.cpp
--- a/Engine/Source/Runtime/Engine/Private/AudioDecompress.cpp
+++ b/Engine/Source/Runtime/Engine/Private/AudioDecompress.cpp
@@ -629,7 +629,7 @@ void FAsyncAudioDecompressWorker::DoWork()
 		{
 			FScopeCycleCounterUObject WaveObject( Wave );
 
-#if PLATFORM_ANDROID
+#if PLATFORM_ANDROID && !PLATFORM_LUMIN
 			// Handle resampling
 			if (QualityInfo.SampleRate > 48000)
 			{
