#DEV-VR - Fixing MagicLeapSDK Settings to be globaluserconfig (by removing the defaultconfig).

diff --git a/Engine/Source/Developer/Lumin/LuminPlatformEditor/Private/MagicLeapSDKSettings.h b/Engine/Source/Developer/Lumin/LuminPlatformEditor/Private/MagicLeapSDKSettings.h
--- a/Engine/Source/Developer/Lumin/LuminPlatformEditor/Private/MagicLeapSDKSettings.h
+++ b/Engine/Source/Developer/Lumin/LuminPlatformEditor/Private/MagicLeapSDKSettings.h
@@ -11,7 +11,7 @@ class IAndroidDeviceDetection;
 /**
  * Implements the settings for the Magic Leap SDK setup.
  */
-UCLASS(config = Engine, globaluserconfig, defaultconfig)
+UCLASS(config = Engine, globaluserconfig)
 // Use defaultconfig here ^ so that MLSDKPath gets updated in the DefaultEngine.ini and matched with the config var used by UEBuildLumin.cs and MagicLeapPluginUtil.h
 class LUMINPLATFORMEDITOR_API UMagicLeapSDKSettings : public UObject
 {
