#DEV-VR - removing stale comment about ML SDK Settings using defaultengine.ini

diff --git a/Engine/Source/Developer/Lumin/LuminPlatformEditor/Private/MagicLeapSDKSettings.h b/Engine/Source/Developer/Lumin/LuminPlatformEditor/Private/MagicLeapSDKSettings.h
--- a/Engine/Source/Developer/Lumin/LuminPlatformEditor/Private/MagicLeapSDKSettings.h
+++ b/Engine/Source/Developer/Lumin/LuminPlatformEditor/Private/MagicLeapSDKSettings.h
@@ -12,7 +12,6 @@ class IAndroidDeviceDetection;
  * Implements the settings for the Magic Leap SDK setup.
  */
 UCLASS(config = Engine, globaluserconfig)
-// Use defaultconfig here ^ so that MLSDKPath gets updated in the DefaultEngine.ini and matched with the config var used by UEBuildLumin.cs and MagicLeapPluginUtil.h
 class LUMINPLATFORMEDITOR_API UMagicLeapSDKSettings : public UObject
 {
 public:
