- Using new multiple-platform Config system to let MeshLOD and Texture cooking system to get ini settings as needed, instead of having each TargetPlatform LoadLocalIni (which was a cause of some bugs)
- Removed other uses of LoadLocalIni from TargetPlatforms (eventually LoadLocalIni and LoadGlobalIni will hopefully be hidden away with better APIs in FConfigCacheIni)

diff --git a/Engine/Plugins/Runtime/AR/Microsoft/HoloLensAR/Source/HoloLensTargetPlatform/Private/HoloLensTargetPlatform.cpp b/Engine/Plugins/Runtime/AR/Microsoft/HoloLensAR/Source/HoloLensTargetPlatform/Private/HoloLensTargetPlatform.cpp
--- a/Engine/Plugins/Runtime/AR/Microsoft/HoloLensAR/Source/HoloLensTargetPlatform/Private/HoloLensTargetPlatform.cpp
+++ b/Engine/Plugins/Runtime/AR/Microsoft/HoloLensAR/Source/HoloLensTargetPlatform/Private/HoloLensTargetPlatform.cpp
@@ -20,9 +20,8 @@ FHoloLensTargetPlatform::FHoloLensTargetPlatform(bool bIsClientOnly)
 	: TNonDesktopTargetPlatformBase(bIsClientOnly)
 {
 #if WITH_ENGINE
-	FConfigCacheIni::LoadLocalIniFile(EngineSettings, TEXT("Engine"), true, *IniPlatformName());
 	TextureLODSettings = nullptr; // These are registered by the device profile system.
-	StaticMeshLODSettings.Initialize(EngineSettings);
+	StaticMeshLODSettings.Initialize(this);
 #endif
 
 	DeviceDetector = IHoloLensDeviceDetector::Create();
@@ -107,7 +106,7 @@ void FHoloLensTargetPlatform::GetReflectionCaptureFormats(TArray<FName>& OutForm
 
 void FHoloLensTargetPlatform::GetTextureFormats(const UTexture* InTexture, TArray< TArray<FName> >& OutFormats) const
 {
-	GetDefaultTextureFormatNamePerLayer(OutFormats.AddDefaulted_GetRef(), this, InTexture, EngineSettings, false);
+	GetDefaultTextureFormatNamePerLayer(OutFormats.AddDefaulted_GetRef(), this, InTexture, false);
 }
 
 void FHoloLensTargetPlatform::GetAllTextureFormats(TArray<FName>& OutFormats) const
diff --git a/Engine/Plugins/Runtime/AR/Microsoft/HoloLensAR/Source/HoloLensTargetPlatform/Private/HoloLensTargetPlatform.h b/Engine/Plugins/Runtime/AR/Microsoft/HoloLensAR/Source/HoloLensTargetPlatform/Private/HoloLensTargetPlatform.h
--- a/Engine/Plugins/Runtime/AR/Microsoft/HoloLensAR/Source/HoloLensTargetPlatform/Private/HoloLensTargetPlatform.h
+++ b/Engine/Plugins/Runtime/AR/Microsoft/HoloLensAR/Source/HoloLensTargetPlatform/Private/HoloLensTargetPlatform.h
@@ -123,9 +123,6 @@ private:
 	FDelegateHandle DeviceDetectedRegistration;
 
 #if WITH_ENGINE
-	// Holds the Engine INI settings (for quick access).
-	FConfigFile EngineSettings;
-
 	// Holds a cache of the target LOD settings.
 	const UTextureLODSettings* TextureLODSettings;
 
diff --git a/Engine/Source/Developer/Android/AndroidTargetPlatform/Private/AndroidTargetPlatform.cpp b/Engine/Source/Developer/Android/AndroidTargetPlatform/Private/AndroidTargetPlatform.cpp
--- a/Engine/Source/Developer/Android/AndroidTargetPlatform/Private/AndroidTargetPlatform.cpp
+++ b/Engine/Source/Developer/Android/AndroidTargetPlatform/Private/AndroidTargetPlatform.cpp
@@ -195,11 +195,10 @@ FAndroidTargetPlatform::FAndroidTargetPlatform(bool bInIsClient, const TCHAR* Fl
 	, DeviceDetection(nullptr)
 
 {
-	#if WITH_ENGINE
-		FConfigCacheIni::LoadLocalIniFile(EngineSettings, TEXT("Engine"), true, *IniPlatformName());
-			TextureLODSettings = nullptr; // These are registered by the device profile system.
-		StaticMeshLODSettings.Initialize(EngineSettings);
-	#endif
+#if WITH_ENGINE
+	TextureLODSettings = nullptr; // These are registered by the device profile system.
+	StaticMeshLODSettings.Initialize(this);
+#endif
 
 	TickDelegate = FTickerDelegate::CreateRaw(this, &FAndroidTargetPlatform::HandleTicker);
 	TickDelegateHandle = FTicker::GetCoreTicker().AddTicker(TickDelegate, 4.0f);
diff --git a/Engine/Source/Developer/Android/AndroidTargetPlatform/Private/AndroidTargetPlatform.h b/Engine/Source/Developer/Android/AndroidTargetPlatform/Private/AndroidTargetPlatform.h
--- a/Engine/Source/Developer/Android/AndroidTargetPlatform/Private/AndroidTargetPlatform.h
+++ b/Engine/Source/Developer/Android/AndroidTargetPlatform/Private/AndroidTargetPlatform.h
@@ -229,11 +229,6 @@ protected:
 	bool SupportsLandscapeMeshLODStreaming() const;
 	bool SupportsVulkanSM5() const;
 
-#if WITH_ENGINE
-	// Holds the Engine INI settings (for quick access).
-	FConfigFile EngineSettings;
-#endif //WITH_ENGINE
-
 protected:
 
 	// Handles when the ticker fires.
@@ -349,7 +344,7 @@ public:
 			{ { FName(TEXT("AutoDXT")) },	{ FName(TEXT("ASTC_RGBAuto")) } },
 		};
 
-		GetDefaultTextureFormatNamePerLayer(OutFormats.AddDefaulted_GetRef(), this, Texture, EngineSettings, true, false, 1);
+		GetDefaultTextureFormatNamePerLayer(OutFormats.AddDefaulted_GetRef(), this, Texture, true, false, 1);
 
 		for (FName& TextureFormatName : OutFormats.Last())
 		{
diff --git a/Engine/Source/Developer/IOS/IOSTargetPlatform/Private/IOSTargetPlatform.cpp b/Engine/Source/Developer/IOS/IOSTargetPlatform/Private/IOSTargetPlatform.cpp
--- a/Engine/Source/Developer/IOS/IOSTargetPlatform/Private/IOSTargetPlatform.cpp
+++ b/Engine/Source/Developer/IOS/IOSTargetPlatform/Private/IOSTargetPlatform.cpp
@@ -32,9 +32,8 @@ FIOSTargetPlatform::FIOSTargetPlatform(bool bInIsTVOS, bool bIsClientOnly)
 	, bIsTVOS(bInIsTVOS)
 {
 #if WITH_ENGINE
-	FConfigCacheIni::LoadLocalIniFile(EngineSettings, TEXT("Engine"), true, *IniPlatformName());
 	TextureLODSettings = nullptr; // TextureLODSettings are registered by the device profile.
-	StaticMeshLODSettings.Initialize(EngineSettings);
+	StaticMeshLODSettings.Initialize(this);
 #endif // #if WITH_ENGINE
 
 	// initialize the connected device detector
@@ -599,7 +598,7 @@ void FIOSTargetPlatform::GetTextureFormats( const UTexture* Texture, TArray< TAr
         {
             BlockSize = 1;
         }
-		GetDefaultTextureFormatNamePerLayer(TextureFormatNames, this, Texture, EngineSettings, true, false, BlockSize);
+		GetDefaultTextureFormatNamePerLayer(TextureFormatNames, this, Texture, true, false, BlockSize);
 	}
 
 	// include the formats we want (use ASTC first so that it is preferred at runtime if they both exist and it's supported)
diff --git a/Engine/Source/Developer/IOS/IOSTargetPlatform/Private/IOSTargetPlatform.h b/Engine/Source/Developer/IOS/IOSTargetPlatform/Private/IOSTargetPlatform.h
--- a/Engine/Source/Developer/IOS/IOSTargetPlatform/Private/IOSTargetPlatform.h
+++ b/Engine/Source/Developer/IOS/IOSTargetPlatform/Private/IOSTargetPlatform.h
@@ -130,9 +130,6 @@ private:
 	TSharedPtr<FMessageEndpoint, ESPMode::ThreadSafe> MessageEndpoint;
 
 #if WITH_ENGINE
-	// Holds the Engine INI settings, for quick use.
-	FConfigFile EngineSettings;
-
 	// Holds the cache of the target LOD settings.
 	const UTextureLODSettings* TextureLODSettings;
 
diff --git a/Engine/Source/Developer/Linux/LinuxTargetPlatform/Private/LinuxTargetPlatform.h b/Engine/Source/Developer/Linux/LinuxTargetPlatform/Private/LinuxTargetPlatform.h
--- a/Engine/Source/Developer/Linux/LinuxTargetPlatform/Private/LinuxTargetPlatform.h
+++ b/Engine/Source/Developer/Linux/LinuxTargetPlatform/Private/LinuxTargetPlatform.h
@@ -56,9 +56,8 @@ public:
 #endif
 
 #if WITH_ENGINE
-		FConfigCacheIni::LoadLocalIniFile(EngineSettings, TEXT("Engine"), true, *this->IniPlatformName());
 		TextureLODSettings = nullptr;
-		StaticMeshLODSettings.Initialize(EngineSettings);
+		StaticMeshLODSettings.Initialize(this);
 
 		InitDevicesFromConfig();
 
@@ -297,7 +296,7 @@ public:
 		if (!TProperties::IsServerOnly())
 		{
 			// just use the standard texture format name for this texture
-			GetDefaultTextureFormatNamePerLayer(OutFormats.AddDefaulted_GetRef(), this, InTexture, EngineSettings, true);
+			GetDefaultTextureFormatNamePerLayer(OutFormats.AddDefaulted_GetRef(), this, InTexture, true);
 		}
 	}
 
@@ -471,9 +470,6 @@ protected:
 
 
 #if WITH_ENGINE
-	// Holds the Engine INI settings for quick use.
-	FConfigFile EngineSettings;
-
 	// Holds the texture LOD settings.
 	const UTextureLODSettings* TextureLODSettings;
 
diff --git a/Engine/Source/Developer/Lumin/LuminTargetPlatform/Private/LuminTargetPlatform.cpp b/Engine/Source/Developer/Lumin/LuminTargetPlatform/Private/LuminTargetPlatform.cpp
--- a/Engine/Source/Developer/Lumin/LuminTargetPlatform/Private/LuminTargetPlatform.cpp
+++ b/Engine/Source/Developer/Lumin/LuminTargetPlatform/Private/LuminTargetPlatform.cpp
@@ -61,21 +61,21 @@ int32 FLuminTargetPlatform::CheckRequirements(bool bProjectHasCode, EBuildConfig
 bool FLuminTargetPlatform::SupportsMobileRendering() const
 {
 	bool bUseMobileRendering = true;
-	LuminEngineSettings.GetBool(TEXT("/Script/LuminRuntimeSettings.LuminRuntimeSettings"), TEXT("bUseMobileRendering"), bUseMobileRendering);
+	GetConfigSystem()->GetBool(TEXT("/Script/LuminRuntimeSettings.LuminRuntimeSettings"), TEXT("bUseMobileRendering"), bUseMobileRendering, GEngineIni);
 	return bUseMobileRendering;
 }
 
 bool FLuminTargetPlatform::SupportsDesktopRendering() const
 {
 	bool bUseMobileRendering = true;
-	LuminEngineSettings.GetBool(TEXT("/Script/LuminRuntimeSettings.LuminRuntimeSettings"), TEXT("bUseMobileRendering"), bUseMobileRendering);
+	GetConfigSystem()->GetBool(TEXT("/Script/LuminRuntimeSettings.LuminRuntimeSettings"), TEXT("bUseMobileRendering"), bUseMobileRendering, GEngineIni);
 	return bUseMobileRendering == false;
 }
 
-static bool LuminSupportsVulkan(const FConfigFile& InLuminEngineSettings)
+static bool LuminSupportsVulkan(FConfigCacheIni* ConfigSystem)
 {
 	bool bSupportsVulkan = false;
-	InLuminEngineSettings.GetBool(TEXT("/Script/LuminRuntimeSettings.LuminRuntimeSettings"), TEXT("bUseVulkan"), bSupportsVulkan);
+	ConfigSystem->GetBool(TEXT("/Script/LuminRuntimeSettings.LuminRuntimeSettings"), TEXT("bUseVulkan"), bSupportsVulkan, GEngineIni);
 	return bSupportsVulkan;
 }
 
@@ -84,17 +84,6 @@ void FLuminTargetPlatform::RefreshSettings()
 #if WITH_ENGINE
 	//UE_LOG(LogCore, Warning, TEXT("*** DIAGNOSE - REFRESH!!!"));
 
-	//the load above does not move settings from "SourceConfig" member to the object itself.  New loads will do that.
-	FConfigFile NewEngineSettings;
-	FConfigCacheIni::LoadLocalIniFile(NewEngineSettings, TEXT("Engine"), true, *IniPlatformName(), true);
-	LuminEngineSettings = NewEngineSettings;
-	//remove the source config as it is just a copied pointer
-	LuminEngineSettings.SourceConfigFile = nullptr;
-	//override the android version too
-	EngineSettings = NewEngineSettings;
-	//remove the source config as it is just a copied pointer
-	EngineSettings.SourceConfigFile = nullptr;
-
 	// Get the Target RHIs for this platform, we do not always want all those that are supported.
 	TArray<FName> TargetedShaderFormats;
 	GetAllTargetedShaderFormats(TargetedShaderFormats);
@@ -188,7 +177,7 @@ void FLuminTargetPlatform::GetAllPossibleShaderFormats( TArray<FName>& OutFormat
 
 	if (SupportsMobileRendering())
 	{
-		if (LuminSupportsVulkan(LuminEngineSettings))
+		if (LuminSupportsVulkan(GetConfigSystem()))
 		{
 			OutFormats.AddUnique(bUseNOUB ? NAME_VULKAN_ES31_LUMIN_NOUB : NAME_VULKAN_ES31_LUMIN);
 		}
@@ -200,7 +189,7 @@ void FLuminTargetPlatform::GetAllPossibleShaderFormats( TArray<FName>& OutFormat
 
 	if (SupportsDesktopRendering())
 	{
-		if (LuminSupportsVulkan(LuminEngineSettings))
+		if (LuminSupportsVulkan(GetConfigSystem()))
 		{
 			OutFormats.AddUnique(bUseNOUB ? NAME_VULKAN_SM5_LUMIN_NOUB : NAME_VULKAN_SM5_LUMIN);
 		}
@@ -229,7 +218,7 @@ void FLuminTargetPlatform::GetTextureFormats(const UTexture* InTexture, TArray<
 {
 	check(InTexture);
 
-	GetDefaultTextureFormatNamePerLayer(OutFormats.AddDefaulted_GetRef(), this, InTexture, LuminEngineSettings, false);
+	GetDefaultTextureFormatNamePerLayer(OutFormats.AddDefaulted_GetRef(), this, InTexture, false);
 	for (FName& TextureFormatName : OutFormats.Last())
 	{
 		// forward rendering only needs one channel for shadow maps
diff --git a/Engine/Source/Developer/Lumin/LuminTargetPlatform/Private/LuminTargetPlatform.h b/Engine/Source/Developer/Lumin/LuminTargetPlatform/Private/LuminTargetPlatform.h
--- a/Engine/Source/Developer/Lumin/LuminTargetPlatform/Private/LuminTargetPlatform.h
+++ b/Engine/Source/Developer/Lumin/LuminTargetPlatform/Private/LuminTargetPlatform.h
@@ -87,10 +87,5 @@ public:
 
 protected:
 	virtual FAndroidTargetDevicePtr CreateTargetDevice(const ITargetPlatform& InTargetPlatform, const FString& InSerialNumber, const FString& InAndroidVariant) const override;
-
-protected:
-	// Holds the Engine INI settings, for quick use.
-	FConfigFile LuminEngineSettings;
-
 };
 
diff --git a/Engine/Source/Developer/Mac/MacTargetPlatform/Private/GenericMacTargetPlatform.h b/Engine/Source/Developer/Mac/MacTargetPlatform/Private/GenericMacTargetPlatform.h
--- a/Engine/Source/Developer/Mac/MacTargetPlatform/Private/GenericMacTargetPlatform.h
+++ b/Engine/Source/Developer/Mac/MacTargetPlatform/Private/GenericMacTargetPlatform.h
@@ -45,9 +45,8 @@ public:
 #endif
 
 		#if WITH_ENGINE
-			FConfigCacheIni::LoadLocalIniFile(EngineSettings, TEXT("Engine"), true, *this->IniPlatformName());
 			TextureLODSettings = nullptr;
-			StaticMeshLODSettings.Initialize(EngineSettings);
+			StaticMeshLODSettings.Initialize(this);
 		#endif
 	}
 
@@ -159,7 +158,7 @@ return TSuper::SupportsFeature(Feature);
 		if (!IS_DEDICATED_SERVER)
 		{
 			// just use the standard texture format name for this texture (with DX11 support)
-			GetDefaultTextureFormatNamePerLayer(OutFormats.AddDefaulted_GetRef(), this, Texture, EngineSettings, true);
+			GetDefaultTextureFormatNamePerLayer(OutFormats.AddDefaulted_GetRef(), this, Texture, true);
 		}
 	}
 
@@ -200,7 +199,7 @@ return TSuper::SupportsFeature(Feature);
 
 		bool bUseDXT5NormalMap = false;
 		FString UseDXT5NormalMapsString;
-		if (EngineSettings.GetString(TEXT("SystemSettings"), TEXT("Compat.UseDXT5NormalMaps"), UseDXT5NormalMapsString))
+		if (this->GetConfigSystem()->GetString(TEXT("SystemSettings"), TEXT("Compat.UseDXT5NormalMaps"), UseDXT5NormalMapsString, GEngineIni))
 		{
 			bUseDXT5NormalMap = FCString::ToBool(*UseDXT5NormalMapsString);
 		}
@@ -366,9 +365,6 @@ private:
 	ITargetDevicePtr LocalDevice;
 
 #if WITH_ENGINE
-	// Holds the Engine INI settings for quick use.
-	FConfigFile EngineSettings;
-
 	// Holds the texture LOD settings.
 	const UTextureLODSettings* TextureLODSettings;
 
diff --git a/Engine/Source/Developer/TargetPlatform/Public/Common/TargetPlatformBase.h b/Engine/Source/Developer/TargetPlatform/Public/Common/TargetPlatformBase.h
--- a/Engine/Source/Developer/TargetPlatform/Public/Common/TargetPlatformBase.h
+++ b/Engine/Source/Developer/TargetPlatform/Public/Common/TargetPlatformBase.h
@@ -41,6 +41,11 @@ public:
 		return *PlatformInfo->DataDrivenPlatformInfo;
 	}
 
+	virtual FConfigCacheIni* GetConfigSystem() const override
+	{
+		return FConfigCacheIni::ForPlatform(*IniPlatformName());
+	}
+
 	TARGETPLATFORM_API virtual bool UsesForwardShading() const override;
 
 	TARGETPLATFORM_API virtual bool UsesDBuffer() const override;
diff --git a/Engine/Source/Developer/TargetPlatform/Public/Interfaces/ITargetPlatform.h b/Engine/Source/Developer/TargetPlatform/Public/Interfaces/ITargetPlatform.h
--- a/Engine/Source/Developer/TargetPlatform/Public/Interfaces/ITargetPlatform.h
+++ b/Engine/Source/Developer/TargetPlatform/Public/Interfaces/ITargetPlatform.h
@@ -225,6 +225,11 @@ public:
 	 */
 	virtual const FDataDrivenPlatformInfo& GetPlatformInfo() const = 0;
 
+	/**
+	 * Returns the config system object usable by this TargetPlatform. It should not be modified in anyway
+	 */
+	virtual FConfigCacheIni* GetConfigSystem() const = 0;
+
 	/**
 	 * Gets the platform's INI name (so an offline tool can load the INI for the given target platform).
 	 *
diff --git a/Engine/Source/Developer/Windows/WindowsTargetPlatform/Private/GenericWindowsTargetPlatform.h b/Engine/Source/Developer/Windows/WindowsTargetPlatform/Private/GenericWindowsTargetPlatform.h
--- a/Engine/Source/Developer/Windows/WindowsTargetPlatform/Private/GenericWindowsTargetPlatform.h
+++ b/Engine/Source/Developer/Windows/WindowsTargetPlatform/Private/GenericWindowsTargetPlatform.h
@@ -42,9 +42,8 @@ public:
 #endif
 
 	#if WITH_ENGINE
-		FConfigCacheIni::LoadLocalIniFile(EngineSettings, TEXT("Engine"), true, *this->IniPlatformName());
 		TextureLODSettings = nullptr; // These are registered by the device profile system.
-		StaticMeshLODSettings.Initialize(EngineSettings);
+		StaticMeshLODSettings.Initialize(this);
 
 
 		// Get the Target RHIs for this platform, we do not always want all those that are supported.
@@ -251,7 +250,7 @@ public:
 	{
 		if (!IS_DEDICATED_SERVER)
 		{
-			GetDefaultTextureFormatNamePerLayer(OutFormats.AddDefaulted_GetRef(), this, InTexture, EngineSettings, bSupportDX11TextureFormats, bSupportCompressedVolumeTexture);
+			GetDefaultTextureFormatNamePerLayer(OutFormats.AddDefaulted_GetRef(), this, InTexture, bSupportDX11TextureFormats, bSupportCompressedVolumeTexture);
 		}
 	}
 
@@ -290,7 +289,7 @@ public:
 
 		bool bUseDXT5NormalMap = false;
 		FString UseDXT5NormalMapsString;
-		if (EngineSettings.GetString(TEXT("SystemSettings"), TEXT("Compat.UseDXT5NormalMaps"), UseDXT5NormalMapsString))
+		if (this->GetConfigSystem()->GetString(TEXT("SystemSettings"), TEXT("Compat.UseDXT5NormalMaps"), UseDXT5NormalMapsString, GEngineIni))
 		{
 			bUseDXT5NormalMap = FCString::ToBool(*UseDXT5NormalMapsString);
 		}
@@ -464,9 +463,6 @@ private:
 	ITargetDevicePtr LocalDevice;
 
 #if WITH_ENGINE
-	// Holds the Engine INI settings for quick use.
-	FConfigFile EngineSettings;
-
 	// Holds the texture LOD settings.
 	const UTextureLODSettings* TextureLODSettings;
 
diff --git a/Engine/Source/Runtime/Engine/Private/StaticMesh.cpp b/Engine/Source/Runtime/Engine/Private/StaticMesh.cpp
--- a/Engine/Source/Runtime/Engine/Private/StaticMesh.cpp
+++ b/Engine/Source/Runtime/Engine/Private/StaticMesh.cpp
@@ -1962,7 +1962,7 @@ void FStaticMeshRenderData::SyncUVChannelData(const TArray<FStaticMaterial>& Obj
 	FStaticMeshLODSettings
 ------------------------------------------------------------------------------*/
 
-void FStaticMeshLODSettings::Initialize(const FConfigFile& IniFile)
+void FStaticMeshLODSettings::Initialize(const ITargetPlatform* TargetPlatform)
 {
 	check(!Groups.Num());
 	// Ensure there is a default LOD group.
@@ -1970,7 +1970,7 @@ void FStaticMeshLODSettings::Initialize(const FConfigFile& IniFile)
 
 	// Read individual entries from a config file.
 	const TCHAR* IniSection = TEXT("StaticMeshLODSettings");
-	const FConfigSection* Section = IniFile.Find(IniSection);
+	const FConfigSection* Section = TargetPlatform->GetConfigSystem()->GetSectionPrivate(IniSection, false, false, GEngineIni);
 	if (Section)
 	{
 		for (TMultiMap<FName,FConfigValue>::TConstIterator It(*Section); It; ++It)
diff --git a/Engine/Source/Runtime/Engine/Private/Texture.cpp b/Engine/Source/Runtime/Engine/Private/Texture.cpp
--- a/Engine/Source/Runtime/Engine/Private/Texture.cpp
+++ b/Engine/Source/Runtime/Engine/Private/Texture.cpp
@@ -1506,7 +1506,7 @@ void UTexture::SetLayerFormatSettings(int32 LayerIndex, const FTextureFormatSett
 
 #endif // #if WITH_EDITOR
 
-FName GetDefaultTextureFormatName( const ITargetPlatform* TargetPlatform, const UTexture* Texture, int32 LayerIndex, const FConfigFile& EngineSettings, bool bSupportDX11TextureFormats, bool bSupportCompressedVolumeTexture, int32 BlockSize )
+FName GetDefaultTextureFormatName( const ITargetPlatform* TargetPlatform, const UTexture* Texture, int32 LayerIndex, bool bSupportDX11TextureFormats, bool bSupportCompressedVolumeTexture, int32 BlockSize )
 {
 	FName TextureFormatName = NAME_None;
 
@@ -1568,7 +1568,7 @@ FName GetDefaultTextureFormatName( const ITargetPlatform* TargetPlatform, const
 
 	FString UseDXT5NormalMapsString;
 
-	if (EngineSettings.GetString(TEXT("SystemSettings"), TEXT("Compat.UseDXT5NormalMaps"), UseDXT5NormalMapsString))
+	if (TargetPlatform->GetConfigSystem()->GetString(TEXT("SystemSettings"), TEXT("Compat.UseDXT5NormalMaps"), UseDXT5NormalMapsString, GEngineIni))
 	{
 		bUseDXT5NormalMap = FCString::ToBool(*UseDXT5NormalMapsString);
 	}
@@ -1677,13 +1677,13 @@ FName GetDefaultTextureFormatName( const ITargetPlatform* TargetPlatform, const
 	return TextureFormatName;
 }
 
-void GetDefaultTextureFormatNamePerLayer(TArray<FName>& OutFormatNames, const class ITargetPlatform* TargetPlatform, const class UTexture* Texture, const class FConfigFile& EngineSettings, bool bSupportDX11TextureFormats, bool bSupportCompressedVolumeTexture, int32 BlockSize)
+void GetDefaultTextureFormatNamePerLayer(TArray<FName>& OutFormatNames, const class ITargetPlatform* TargetPlatform, const class UTexture* Texture, bool bSupportDX11TextureFormats, bool bSupportCompressedVolumeTexture, int32 BlockSize)
 {
 #if WITH_EDITOR
 	OutFormatNames.Reserve(Texture->Source.GetNumLayers());
 	for (int32 LayerIndex = 0; LayerIndex < Texture->Source.GetNumLayers(); ++LayerIndex)
 	{
-		OutFormatNames.Add(GetDefaultTextureFormatName(TargetPlatform, Texture, LayerIndex, EngineSettings, bSupportDX11TextureFormats, bSupportCompressedVolumeTexture, BlockSize));
+		OutFormatNames.Add(GetDefaultTextureFormatName(TargetPlatform, Texture, LayerIndex, bSupportDX11TextureFormats, bSupportCompressedVolumeTexture, BlockSize));
 	}
 #endif // WITH_EDITOR
 }
diff --git a/Engine/Source/Runtime/Engine/Public/StaticMeshResources.h b/Engine/Source/Runtime/Engine/Public/StaticMeshResources.h
--- a/Engine/Source/Runtime/Engine/Public/StaticMeshResources.h
+++ b/Engine/Source/Runtime/Engine/Public/StaticMeshResources.h
@@ -141,7 +141,7 @@ public:
 	 * Initializes LOD settings by reading them from the passed in config file section.
 	 * @param IniFile Preloaded ini file object to load from
 	 */
-	ENGINE_API void Initialize(const FConfigFile& IniFile);
+	ENGINE_API void Initialize(const ITargetPlatform* TargetPlatform);
 
 	/** Retrieve the settings for the specified LOD group. */
 	const FStaticMeshLODGroup& GetLODGroup(FName LODGroup) const
diff --git a/Engine/Source/Runtime/Engine/Public/TextureResource.h b/Engine/Source/Runtime/Engine/Public/TextureResource.h
--- a/Engine/Source/Runtime/Engine/Public/TextureResource.h
+++ b/Engine/Source/Runtime/Engine/Public/TextureResource.h
@@ -531,10 +531,10 @@ private:
 };
 
 /** Gets the name of a format for the given LayerIndex */
-ENGINE_API FName GetDefaultTextureFormatName( const class ITargetPlatform* TargetPlatform, const class UTexture* Texture, int32 LayerIndex, const class FConfigFile& EngineSettings, bool bSupportDX11TextureFormats, bool bSupportCompressedVolumeTexture = false, int32 BlockSize = 4);
+ENGINE_API FName GetDefaultTextureFormatName( const class ITargetPlatform* TargetPlatform, const class UTexture* Texture, int32 LayerIndex, bool bSupportDX11TextureFormats, bool bSupportCompressedVolumeTexture = false, int32 BlockSize = 4);
 
 /** Gets an array of format names for each layer in the texture */
-ENGINE_API void GetDefaultTextureFormatNamePerLayer(TArray<FName>& OutFormatNames, const class ITargetPlatform* TargetPlatform, const class UTexture* Texture, const class FConfigFile& EngineSettings, bool bSupportDX11TextureFormats, bool bSupportCompressedVolumeTexture = false, int32 BlockSize = 4);
+ENGINE_API void GetDefaultTextureFormatNamePerLayer(TArray<FName>& OutFormatNames, const class ITargetPlatform* TargetPlatform, const class UTexture* Texture, bool bSupportDX11TextureFormats, bool bSupportCompressedVolumeTexture = false, int32 BlockSize = 4);
 
 // returns all the texture formats which can be returned by GetDefaultTextureFormatName
 ENGINE_API void GetAllDefaultTextureFormats( const class ITargetPlatform* TargetPlatform, TArray<FName>& OutFormats, bool bSupportDX11TextureFormats);
